Remove UE4 references

#android

diff --git a/Engine/Plugins/Online/Android/OnlineSubsystemGameCircle/Build/Android/Java/src/com/epicgames/ue4/AmazonStoreHelper.java b/Engine/Plugins/Online/Android/OnlineSubsystemGameCircle/Build/Android/Java/src/com/epicgames/ue4/AmazonStoreHelper.java
deleted file mode 100644
diff --git a/Engine/Plugins/Online/Android/OnlineSubsystemGameCircle/Build/Android/Java/src/com/epicgames/unreal/AmazonStoreHelper.java b/Engine/Plugins/Online/Android/OnlineSubsystemGameCircle/Build/Android/Java/src/com/epicgames/unreal/AmazonStoreHelper.java
new file mode 100644
--- /dev/null
+++ b/Engine/Plugins/Online/Android/OnlineSubsystemGameCircle/Build/Android/Java/src/com/epicgames/unreal/AmazonStoreHelper.java
@@ -0,0 +1,323 @@
+// Copyright Epic Games, Inc. All Rights Reserved.
+
+package com.epicgames.unreal;
+
+import java.util.Map;
+import java.util.HashSet;
+import java.util.Set;
+import java.util.ArrayList;
+import java.util.List;
+
+import android.content.Intent;
+
+import com.amazon.device.iap.PurchasingListener;
+import com.amazon.device.iap.PurchasingService;
+
+import com.amazon.device.iap.model.FulfillmentResult;
+import com.amazon.device.iap.model.Product;
+import com.amazon.device.iap.model.ProductDataResponse;
+import com.amazon.device.iap.model.ProductType;
+import com.amazon.device.iap.model.PurchaseResponse;
+import com.amazon.device.iap.model.PurchaseUpdatesResponse;
+import com.amazon.device.iap.model.Receipt;
+import com.amazon.device.iap.model.RequestId;
+import com.amazon.device.iap.model.UserData;
+import com.amazon.device.iap.model.UserDataResponse;
+
+
+public class AmazonStoreHelper implements StoreHelper, PurchasingListener
+{
+	// Output device for log messages.
+	private Logger Log;
+
+	// Cache access to the games activity.
+	private GameActivity gameActivity;
+
+	// Set to true after everything is set up
+	private boolean bIsAllowedToMakePurchases;
+
+	// Is a query for IAP in progress
+	private boolean bQueryInProgress;
+
+	// Product Id Query
+	private Set<String> ProductSkus;
+
+	// Restore purchases receipt data
+	private List<Receipt> RestoreReceipts;
+	// True when a restore of purchases is in progress
+	private boolean bRestoreInProgress;
+
+	// Cached user data returned from Amazon
+	private UserData LocalUserData;
+
+	// Save the product ID that was most recently purchased
+	private String MostRecentPurchaseSku;
+
+	public enum EAmazonResponseStatus
+	{
+		Successful,
+		Failed,
+		NotSupported,
+		AlreadyPurchased,
+		InvalidSKU,
+		Unknown
+	}
+
+
+	public AmazonStoreHelper(String InProductKey, GameActivity InGameActivity, final Logger InLog)
+	{
+		bIsAllowedToMakePurchases = false;
+		bQueryInProgress = false;
+		bRestoreInProgress = false;
+
+		gameActivity = InGameActivity;
+		Log = InLog;
+
+		Log.debug("Registering PurchasingListener");
+		PurchasingService.registerListener(gameActivity.getApplicationContext(), this);
+		Log.debug("IS_SANDBOX_MODE = " + PurchasingService.IS_SANDBOX_MODE);
+	}
+
+	public void OnServiceInitialized()
+	{
+		Log.debug("Requesting UserData from PurchasingService");
+		PurchasingService.getUserData();
+	}
+
+
+	//~ Begin StoreHelper interface
+	public boolean QueryInAppPurchases(String[] ProductIDs)
+	{
+		if(ProductSkus != null)
+		{
+			Log.debug("AmazonStoreHelper.QueryInAppPurchases Query Already In Progress");
+			return false;
+		}
+
+		ProductSkus = new HashSet<String>();
+
+		for (String productId : ProductIDs)
+		{
+			Log.debug("AmazonStoreHelper.QueryInAppPurchases Querying - " + productId);
+			ProductSkus.add(productId);
+		}
+
+		PurchasingService.getProductData(ProductSkus);
+
+		return true;
+	}
+
+	public boolean BeginPurchase(String ProductID)
+	{
+		Log.debug("AmazonStoreHelper.BeginPurchase - " + ProductID);
+
+		final RequestId requestId = PurchasingService.purchase(ProductID);
+		Log.debug("AmazonStoreHelper.BeginPurchase RequestId = " + requestId);
+		return true;
+	}
+
+	public boolean IsAllowedToMakePurchases()
+	{
+		return bIsAllowedToMakePurchases;
+	}
+
+	public boolean RestorePurchases(String[] InProductIDs, boolean[] bConsumable)
+	{
+		// Note: ProductIDs and Consumable flags not needed on Amazon since we actually just check the ProductType in onPurchaseUpdatesResponse
+		Log.debug("AmazonStoreHelper.RestorePurchases Calling PurchasingService.getPurchaseUpdates");
+		if(bRestoreInProgress)
+		{
+			Log.debug("AmazonStoreHelper.RestorePurchases - Restore Already In Progress");
+			return false;
+		}
+
+		PurchasingService.getPurchaseUpdates(true);
+		bRestoreInProgress = true;
+		return true;
+	}
+
+	public void	onDestroy()
+	{
+	}
+
+	public boolean onActivityResult(int requestCode, int resultCode, Intent data)
+	{
+		return true;
+	}
+	//~ End StoreHelper interface
+
+	//~ Begin PurchasingListener interface
+	public void onProductDataResponse(ProductDataResponse productDataResponse)
+	{
+		switch(productDataResponse.getRequestStatus())
+		{
+		case SUCCESSFUL:
+			{
+				Log.debug("AmazonStoreHelper.onProductDataResponse Status: SUCCESS");
+				for(final String unavailableSku : productDataResponse.getUnavailableSkus())
+				{
+					Log.debug("Unavailable SKU: " + unavailableSku);
+				}
+
+				final Map<String, Product> products = productDataResponse.getProductData();
+				ArrayList<String> titles = new ArrayList<String>();
+				ArrayList<String> descriptions = new ArrayList<String>();
+				ArrayList<String> prices = new ArrayList<String>();
+
+				for ( final String key : products.keySet()) 
+				{
+					Product product = products.get(key);
+					Log.debug(String.format( "Product: %s\n Type: %s\n SKU: %s\n Price: %s\n Description: %s\n" , product.getTitle(), product.getProductType(), product.getSku(), product.getPrice(), product.getDescription()));
+					
+					titles.add(product.getTitle());
+					descriptions.add(product.getDescription());
+					prices.add(product.getPrice());
+				}
+
+				nativeQueryComplete(EAmazonResponseStatus.Successful.ordinal(), ProductSkus.toArray(new String[ProductSkus.size()]), titles.toArray(new String[titles.size()]), descriptions.toArray(new String[descriptions.size()]), prices.toArray(new String[prices.size()]));
+			}
+			break;
+		case FAILED:
+			Log.debug("AmazonStoreHelper.onProductDataResponse Status: FAILED");
+			nativeQueryComplete(EAmazonResponseStatus.Failed.ordinal(), null, null, null, null);
+			break;
+		case NOT_SUPPORTED:
+			Log.debug("AmazonStoreHelper.onProductDataResponse Status: NOT_SUPPORTED");
+			nativeQueryComplete(EAmazonResponseStatus.NotSupported.ordinal(), null, null, null, null);
+			break;
+		default:
+			Log.debug("AmazonStoreHelper.onProductDataResponse Status unknown: " + productDataResponse.getRequestStatus());
+			nativeQueryComplete(EAmazonResponseStatus.Unknown.ordinal(), null, null, null, null);
+			break;
+		}
+
+		ProductSkus = null;
+	}
+
+	public void onPurchaseResponse(PurchaseResponse purchaseResponse)
+	{
+		switch(purchaseResponse.getRequestStatus())
+		{
+		case SUCCESSFUL:
+			{
+				Log.debug("AmazonStoreHelper.onPurchaseResponse Status: SUCCESSFUL " + purchaseResponse);
+				final Receipt receipt = purchaseResponse.getReceipt();
+
+				LocalUserData = purchaseResponse.getUserData();
+
+				nativePurchaseComplete(EAmazonResponseStatus.Successful.ordinal(), receipt.getSku(), receipt.toJSON().toString(), receipt.getReceiptId());
+
+				PurchasingService.notifyFulfillment(receipt.getReceiptId(), FulfillmentResult.FULFILLED);
+			}
+			break;
+		case ALREADY_PURCHASED:
+			Log.debug("AmazonStoreHelper.onPurchaseResponse Status: ALREADY_PURCHASED");
+			nativePurchaseComplete(EAmazonResponseStatus.AlreadyPurchased.ordinal(), null, null, null);
+			break;
+		case INVALID_SKU:
+			Log.debug("AmazonStoreHelper.onPurchaseResponse Status: INVALID_SKU");
+			nativePurchaseComplete(EAmazonResponseStatus.InvalidSKU.ordinal(), null, null, null);
+			break;
+		case FAILED:
+			Log.debug("AmazonStoreHelper.onPurchaseResponse Status: FAILED");
+			nativePurchaseComplete(EAmazonResponseStatus.Failed.ordinal(), null, null, null);
+			break;
+		case NOT_SUPPORTED:
+			Log.debug("AmazonStoreHelper.onPurchaseResponse Status: NOT_SUPPORTED");
+			nativePurchaseComplete(EAmazonResponseStatus.NotSupported.ordinal(), null, null, null);
+			break;
+		default:
+			Log.debug("AmazonStoreHelper.onPurchaseResponse Status unknown: " + purchaseResponse.getRequestStatus());
+			nativePurchaseComplete(EAmazonResponseStatus.Unknown.ordinal(), null, null, null);
+			break;
+		}
+	}
+
+	public void onPurchaseUpdatesResponse(PurchaseUpdatesResponse purchaseUpdatesResponse)
+	{
+		switch(purchaseUpdatesResponse.getRequestStatus())
+		{
+		case SUCCESSFUL:
+			Log.debug("AmazonStoreHelper.onPurchaseUpdatesResponse Status: SUCCESSFUL");
+			if(RestoreReceipts == null)
+			{
+				RestoreReceipts = purchaseUpdatesResponse.getReceipts();
+			}
+			else
+			{
+				for(final Receipt receipt : purchaseUpdatesResponse.getReceipts())
+				{
+					RestoreReceipts.add(receipt);
+				}
+			}
+
+			if(purchaseUpdatesResponse.hasMore())
+			{
+				PurchasingService.getPurchaseUpdates(false);
+			}
+			else
+			{
+				ArrayList<String> restoreSkus = new ArrayList<String>();
+				ArrayList<String> restoreReceiptData = new ArrayList<String>();
+
+				for(final Receipt receipt : RestoreReceipts)
+				{
+					Log.debug("AmazonStoreHelper.onPurchaseUpdatesResponse SKU: " + receipt.getSku() + " ReceiptJSON: " + receipt.toJSON());
+					restoreSkus.add(receipt.getSku());
+					restoreReceiptData.add(receipt.toJSON().toString());
+
+					//	If we are returned a receipt for a consumable purchase, mark it fulfilled
+					if(receipt.getProductType() == ProductType.CONSUMABLE)
+					{
+						PurchasingService.notifyFulfillment(receipt.getReceiptId(), FulfillmentResult.FULFILLED);
+					}
+				}
+
+				nativeRestorePurchasesComplete(EAmazonResponseStatus.Successful.ordinal(), restoreSkus.toArray( new String[restoreSkus.size()] ), restoreReceiptData.toArray( new String[restoreReceiptData.size()] ) );
+
+				RestoreReceipts = null;
+				bRestoreInProgress = false;
+			}
+			break;
+		case FAILED:
+			Log.debug("AmazonStoreHelper.onPurchaseUpdatesResponse Status: FAILED");
+			nativeRestorePurchasesComplete(EAmazonResponseStatus.Failed.ordinal(), null, null);
+			break;
+		case NOT_SUPPORTED:
+			Log.debug("AmazonStoreHelper.onPurchaseUpdatesResponse Status: NOT_SUPPORTED");
+			nativeRestorePurchasesComplete(EAmazonResponseStatus.NotSupported.ordinal(), null, null);
+			break;
+		default:
+			Log.debug("AmazonStoreHelper.onPurchaseUpdatesResponse Status unknown: " + purchaseUpdatesResponse.getRequestStatus());
+			nativeRestorePurchasesComplete(EAmazonResponseStatus.Unknown.ordinal(), null, null);
+			break;
+		}
+	}
+
+	public void onUserDataResponse(UserDataResponse userDataResponse)
+	{
+		switch(userDataResponse.getRequestStatus())
+		{
+		case SUCCESSFUL:
+			LocalUserData = userDataResponse.getUserData();
+			Log.debug("AmazonStoreHelper.onUserDataResponse Request Successful\nLocalUserData = " + LocalUserData.toString());
+			bIsAllowedToMakePurchases = true;
+			break;
+		case FAILED:
+			Log.debug("AmazonStoreHelper.onUserDataResponse Request Failed");
+			break;
+		case NOT_SUPPORTED:
+			Log.debug("AmazonStoreHelper.onUserDataResponse Request Not Supported");
+			break;
+		default:
+			Log.debug("AmazonStoreHelper.onUserDataResponse Status unknown: " + userDataResponse.getRequestStatus());
+			break;
+		}
+	}
+	//~ End PurchasingListener interface
+
+	// Callback that notify the C++ implementation that a task has completed
+	public native void nativeQueryComplete(int responseStatus, String[] productIDs, String[] titles, String[] descriptions, String[] prices);
+	public native void nativePurchaseComplete(int responseStatus, String ProductId, String ReceiptData, String Signature);
+	public native void nativeRestorePurchasesComplete(int responseStatus, String[] ProductIds, String[] ReceiptsData);
+}
\ No newline at end of file
